Modify reimport analytics to clarify exactly what file type was reimported.

diff --git a/Engine/Source/Editor/DestructibleMeshEditor/Private/DestructibleMeshEditor.cpp b/Engine/Source/Editor/DestructibleMeshEditor/Private/DestructibleMeshEditor.cpp
--- a/Engine/Source/Editor/DestructibleMeshEditor/Private/DestructibleMeshEditor.cpp
+++ b/Engine/Source/Editor/DestructibleMeshEditor/Private/DestructibleMeshEditor.cpp
@@ -8,9 +8,6 @@
 #include "Editor/PropertyEditor/Public/PropertyEditorModule.h"
 #include "Editor/PropertyEditor/Public/IDetailsView.h"
 
-#include "Runtime/Analytics/Analytics/Public/Interfaces/IAnalyticsProvider.h"
-#include "EngineAnalytics.h"
-
 #define LOCTEXT_NAMESPACE "DestructibleMeshEditor"
 
 const FName FDestructibleMeshEditor::ViewportTabId( TEXT( "DestructibleMeshEditor_Viewport" ) );
@@ -398,10 +395,6 @@ void FDestructibleMeshEditor::OnPostReimport(UObject* InObject, bool bSuccess)
 
 	if ( bSuccess )
 	{
-		if (FEngineAnalytics::IsAvailable())
-		{
-			FEngineAnalytics::GetProvider().RecordEvent(TEXT("Editor.Usage.DestructibleMesh.ReimportedViaEditor"));
-		}
 		RefreshTool();
 	}
 }
diff --git a/Engine/Source/Editor/FontEditor/Private/FontEditor.cpp b/Engine/Source/Editor/FontEditor/Private/FontEditor.cpp
--- a/Engine/Source/Editor/FontEditor/Private/FontEditor.cpp
+++ b/Engine/Source/Editor/FontEditor/Private/FontEditor.cpp
@@ -12,9 +12,6 @@
 #include "Editor/PropertyEditor/Public/PropertyEditorModule.h"
 #include "Editor/PropertyEditor/Public/IDetailsView.h"
 
-#include "Runtime/Analytics/Analytics/Public/Interfaces/IAnalyticsProvider.h"
-#include "EngineAnalytics.h"
-
 #define LOCTEXT_NAMESPACE "FontEditor"
 
 DEFINE_LOG_CATEGORY_STATIC(LogFontEditor, Log, All);
@@ -721,10 +718,6 @@ void FFontEditor::OnPostReimport(UObject* InObject, bool bSuccess)
 
 	if ( bSuccess )
 	{
-		if (FEngineAnalytics::IsAvailable())
-		{
-			FEngineAnalytics::GetProvider().RecordEvent(TEXT("Editor.Usage.Font.ReimportedViaEditor"));
-		}
 		FontViewport->RefreshViewport();
 		FontPreviewWidget->RefreshViewport();
 	}
diff --git a/Engine/Source/Editor/Persona/Private/Persona.cpp b/Engine/Source/Editor/Persona/Private/Persona.cpp
--- a/Engine/Source/Editor/Persona/Private/Persona.cpp
+++ b/Engine/Source/Editor/Persona/Private/Persona.cpp
@@ -64,9 +64,6 @@
 #include "AnimationEditorUtils.h"
 #include "SAnimationSequenceBrowser.h"
 
-#include "Runtime/Analytics/Analytics/Public/Interfaces/IAnalyticsProvider.h"
-#include "EngineAnalytics.h"
-
 #define LOCTEXT_NAMESPACE "FPersona"
 
 /////////////////////////////////////////////////////
@@ -909,10 +906,6 @@ void FPersona::OnPostReimport(UObject* InObject, bool bSuccess)
 {
 	if (bSuccess)
 	{
-		if (FEngineAnalytics::IsAvailable())
-		{
-			FEngineAnalytics::GetProvider().RecordEvent(TEXT("Editor.Usage.Persona.ReimportedViaEditor"));
-		}
 		ConditionalRefreshEditor(InObject);
 	}
 }
diff --git a/Engine/Source/Editor/StaticMeshEditor/Private/StaticMeshEditor.cpp b/Engine/Source/Editor/StaticMeshEditor/Private/StaticMeshEditor.cpp
--- a/Engine/Source/Editor/StaticMeshEditor/Private/StaticMeshEditor.cpp
+++ b/Engine/Source/Editor/StaticMeshEditor/Private/StaticMeshEditor.cpp
@@ -1777,10 +1777,6 @@ void FStaticMeshEditor::OnPostReimport(UObject* InObject, bool bSuccess)
 
 	if (bSuccess)
 	{
-		if (FEngineAnalytics::IsAvailable())
-		{
-			FEngineAnalytics::GetProvider().RecordEvent(TEXT("Editor.Usage.StaticMesh.ReimportedViaEditor"));
-		}
 		RefreshTool();
 	}
 }
diff --git a/Engine/Source/Editor/TextureEditor/Private/TextureEditorToolkit.cpp b/Engine/Source/Editor/TextureEditor/Private/TextureEditorToolkit.cpp
--- a/Engine/Source/Editor/TextureEditor/Private/TextureEditorToolkit.cpp
+++ b/Engine/Source/Editor/TextureEditor/Private/TextureEditorToolkit.cpp
@@ -2,8 +2,6 @@
 
 #include "TextureEditorPrivatePCH.h"
 #include "Factories.h"
-#include "Runtime/Analytics/Analytics/Public/Interfaces/IAnalyticsProvider.h"
-#include "EngineAnalytics.h"
 
 #define LOCTEXT_NAMESPACE "FTextureEditorToolkit"
 
@@ -994,14 +992,7 @@ void FTextureEditorToolkit::HandleReimportManagerPostReimport( UObject* InObject
 		return;
 	}
 
-	if (bSuccess)
-	{
-		if (FEngineAnalytics::IsAvailable())
-		{
-			FEngineAnalytics::GetProvider().RecordEvent(TEXT("Editor.Usage.Texture.ReimportedViaEditor"));
-		}
-	}
-	else
+	if (!bSuccess)
 	{
 		// Failed, restore the compression flag
 		Texture->DeferCompression = SavedCompressionSetting;
diff --git a/Engine/Source/Editor/UnrealEd/Private/Editor.cpp b/Engine/Source/Editor/UnrealEd/Private/Editor.cpp
--- a/Engine/Source/Editor/UnrealEd/Private/Editor.cpp
+++ b/Engine/Source/Editor/UnrealEd/Private/Editor.cpp
@@ -2616,7 +2616,9 @@ bool FReimportManager::Reimport( UObject* Obj, bool bAskForNewFileIfMissing )
 						GEditor->BroadcastObjectReimported(Obj);
 						if (FEngineAnalytics::IsAvailable())
 						{
-							FEngineAnalytics::GetProvider().RecordEvent(TEXT("Editor.Usage.AssetReimported"));
+							TArray<FAnalyticsEventAttribute> Attributes;
+							Attributes.Add( FAnalyticsEventAttribute( TEXT( "ObjectType" ), Obj->GetClass()->GetName() ) );
+							FEngineAnalytics::GetProvider().RecordEvent(TEXT("Editor.Usage.AssetReimported"), Attributes);
 						}
 						bSuccess = true;
 					}
