Destructible component will re-cook when chunk parameters are changed. Fixes per chunk settings being ignored.

TTP#339860

diff --git a/Engine/Source/Editor/DestructibleMeshEditor/Classes/DestructibleChunkParamsProxy.h b/Engine/Source/Editor/DestructibleMeshEditor/Classes/DestructibleChunkParamsProxy.h
--- a/Engine/Source/Editor/DestructibleMeshEditor/Classes/DestructibleChunkParamsProxy.h
+++ b/Engine/Source/Editor/DestructibleMeshEditor/Classes/DestructibleChunkParamsProxy.h
@@ -24,5 +24,7 @@ class UDestructibleChunkParamsProxy : public UObject
 	// UObject interface
 	virtual void PostEditChangeProperty(struct FPropertyChangedEvent& PropertyChangedEvent) override;
 	// End of UObject interface
+
+	TWeakPtr<IDestructibleMeshEditor> DestructibleMeshEditorPtr;
 #endif
 };
\ No newline at end of file
diff --git a/Engine/Source/Editor/DestructibleMeshEditor/Private/DestructibleChunkParamsProxy.cpp b/Engine/Source/Editor/DestructibleMeshEditor/Private/DestructibleChunkParamsProxy.cpp
--- a/Engine/Source/Editor/DestructibleMeshEditor/Private/DestructibleChunkParamsProxy.cpp
+++ b/Engine/Source/Editor/DestructibleMeshEditor/Private/DestructibleChunkParamsProxy.cpp
@@ -1,6 +1,7 @@
 // Copyright 1998-2014 Epic Games, Inc. All Rights Reserved.
 
 #include "DestructibleMeshEditorPrivatePCH.h"
+#include "ApexDestructibleAssetImport.h"
 
 UDestructibleChunkParamsProxy::UDestructibleChunkParamsProxy(const class FPostConstructInitializeProperties& PCIP)
 	:Super(PCIP)
@@ -16,4 +17,7 @@ void UDestructibleChunkParamsProxy::PostEditChangeProperty( struct FPropertyChan
 	{
 		DestructibleMesh->FractureSettings->ChunkParameters[ChunkIndex] = ChunkParams;
 	}
+
+	BuildDestructibleMeshFromFractureSettings(*DestructibleMesh, NULL);
+	DestructibleMeshEditorPtr.Pin()->RefreshViewport();
 }
diff --git a/Engine/Source/Editor/DestructibleMeshEditor/Private/SDestructibleMeshEditorViewport.cpp b/Engine/Source/Editor/DestructibleMeshEditor/Private/SDestructibleMeshEditorViewport.cpp
--- a/Engine/Source/Editor/DestructibleMeshEditor/Private/SDestructibleMeshEditorViewport.cpp
+++ b/Engine/Source/Editor/DestructibleMeshEditor/Private/SDestructibleMeshEditorViewport.cpp
@@ -171,6 +171,7 @@ void FDestructibleMeshEditorViewportClient::UpdateChunkSelection( TArray<int32>
 
 		Proxy->DestructibleMesh = DestructibleMesh;
 		Proxy->ChunkIndex = InSelectedChunkIndices[i];
+		Proxy->DestructibleMeshEditorPtr = DestructibleMeshEditorPtr;
 
 		if (FractureSettings != NULL && FractureSettings->ChunkParameters.Num() > Proxy->ChunkIndex)
 		{
