Fix crash when trying to set non-existant tab to be active. UE-11225

diff --git a/Engine/Source/Editor/DestructibleMeshEditor/Private/DestructibleMeshEditor.cpp b/Engine/Source/Editor/DestructibleMeshEditor/Private/DestructibleMeshEditor.cpp
--- a/Engine/Source/Editor/DestructibleMeshEditor/Private/DestructibleMeshEditor.cpp
+++ b/Engine/Source/Editor/DestructibleMeshEditor/Private/DestructibleMeshEditor.cpp
@@ -532,11 +532,17 @@ void FDestructibleMeshEditor::SetSelectedChunks( const TArray<UObject*>& Selecte
 {
 	if (SelectedChunks.Num() > 0)
 	{
-		ChunkParametersViewTab->ActivateInParent(ETabActivationCause::SetDirectly);
+		if(ChunkParametersViewTab.IsValid())
+		{
+			ChunkParametersViewTab->ActivateInParent(ETabActivationCause::SetDirectly);
+		}
 	}
 	else
 	{
-		DestructibleMeshDetailsViewTab->ActivateInParent(ETabActivationCause::SetDirectly);
+		if(DestructibleMeshDetailsViewTab.IsValid())
+		{
+			DestructibleMeshDetailsViewTab->ActivateInParent(ETabActivationCause::SetDirectly);
+		}
 	}
 
 	ChunkParametersView->SetObjects(SelectedChunks, true);
