Fix non apex build

diff --git a/Engine/Source/Editor/DestructibleMeshEditor/Private/DestructibleChunkParamsProxy.cpp b/Engine/Source/Editor/DestructibleMeshEditor/Private/DestructibleChunkParamsProxy.cpp
--- a/Engine/Source/Editor/DestructibleMeshEditor/Private/DestructibleChunkParamsProxy.cpp
+++ b/Engine/Source/Editor/DestructibleMeshEditor/Private/DestructibleChunkParamsProxy.cpp
@@ -19,6 +19,8 @@ void UDestructibleChunkParamsProxy::PostEditChangeProperty( struct FPropertyChan
 		DestructibleMesh->FractureSettings->ChunkParameters[ChunkIndex] = ChunkParams;
 	}
 
+#if WITH_APEX
 	BuildDestructibleMeshFromFractureSettings(*DestructibleMesh, NULL);
+#endif
 	DestructibleMeshEditorPtr.Pin()->RefreshViewport();
 }
